[FlutterCommon] Add connect screen and additional settings widgets

diff --git a/Engine/Extras/FlutterCommon/pubspec.yaml b/Engine/Extras/FlutterCommon/pubspec.yaml
--- a/Engine/Extras/FlutterCommon/pubspec.yaml
+++ b/Engine/Extras/FlutterCommon/pubspec.yaml
@@ -22,6 +22,7 @@ dependencies:
   streaming_shared_preferences: ^2.0.0
   uuid: ^4.2.1
   archive: ^3.2.0
+  share_plus: ^7.2.1
 
 dev_dependencies:
   flutter_test:
