[FlutterCommon] Directly use .arb files for localization instead of generating them from a .csv

diff --git a/Engine/Extras/FlutterCommon/pubspec.yaml b/Engine/Extras/FlutterCommon/pubspec.yaml
--- a/Engine/Extras/FlutterCommon/pubspec.yaml
+++ b/Engine/Extras/FlutterCommon/pubspec.yaml
@@ -17,7 +17,6 @@ dependencies:
   streaming_shared_preferences:
     path: ../ThirdPartyNotUE/streaming_shared_preferences
 
-  arb_generator: ^0.0.3
   flutter_svg: ^2.0.9
   intl: any
   logging: ^1.2.0
